Add a custom IPlatformFile to support bespoke DLC mount paths

#rnx
#Synced-CL 51451766

diff --git a/Engine/Plugins/Experimental/PlatformDLC/Source/PlatformDLC/Private/PlatformDLCModule.cpp b/Engine/Plugins/Experimental/PlatformDLC/Source/PlatformDLC/Private/PlatformDLCModule.cpp
--- a/Engine/Plugins/Experimental/PlatformDLC/Source/PlatformDLC/Private/PlatformDLCModule.cpp
+++ b/Engine/Plugins/Experimental/PlatformDLC/Source/PlatformDLC/Private/PlatformDLCModule.cpp
@@ -1,6 +1,7 @@
 // Copyright Epic Games, Inc. All Rights Reserved.
 
 #include "PlatformDLCModule.h"
+#include "PlatformDLCPlatformFile.h"
 #include "HAL/IConsoleManager.h"
 #include "HAL/PlatformFile.h"
 #include "HAL/PlatformFileManager.h"
@@ -15,6 +16,21 @@ const TCHAR* IPlatformDLC::ConfigSectionName = TEXT("PlatformDLC");
 class FPlatformDLCModule : public IPlatformDLCModule
 {
 public:
+	virtual void StartupModule() override
+	{
+		PlatformDLCFile = IPlatformDLCPlatformFile::Construct();
+	}
+
+	virtual void ShutdownModule() override
+	{
+		// remove DLC notification hook
+		UnregisterForNotifications();
+
+		// remove ourselves from the platform file chain (there can be late writes after the shutdown).
+		UnregisterPlatformFile();
+		PlatformDLCFile.Reset();
+	}
+
 	virtual TSharedPtr<IPlatformDLC> GetPlatformDLC() override
 	{
 		if (!bHasTried)
@@ -27,6 +43,7 @@ public:
 				if (PlatformDLCFactory != nullptr)
 				{
 					PlatformDLC = PlatformDLCFactory->GetPlatformDLC();
+					RegisterForNotifications();
 				}
 			}
 
@@ -36,9 +53,85 @@ public:
 		return PlatformDLC.Pin();
 	}
 
+	virtual IPlatformFile* GetPlatformFile() override
+	{
+		check(PlatformDLCFile.IsValid());
+		return PlatformDLCFile.Get();
+	}
+
 private:
+
+	void RegisterPlatformFile()
+	{
+		verify(PlatformDLCFile->Initialize(&FPlatformFileManager::Get().GetPlatformFile(), nullptr));
+		FPlatformFileManager::Get().SetPlatformFile(*PlatformDLCFile);
+	}
+
+	void UnregisterPlatformFile()
+	{
+		if (FPlatformFileManager::Get().FindPlatformFile(PlatformDLCFile->GetName()))
+		{
+			FPlatformFileManager::Get().RemovePlatformFile(PlatformDLCFile.Get());
+		}
+	}
+
+
+	void RegisterForNotifications()
+	{
+		check(!NotificationDelegate.IsValid());
+
+		TSharedPtr<IPlatformDLC> PlatformDLCStrong = PlatformDLC.Pin();
+		if (PlatformDLCStrong.IsValid())
+		{
+			NotificationDelegate = PlatformDLCStrong->OnNotification().AddRaw(this, &FPlatformDLCModule::OnNotification );
+		}
+	}
+
+	void UnregisterForNotifications()
+	{
+		TSharedPtr<IPlatformDLC> PlatformDLCStrong = PlatformDLC.Pin();
+		if (PlatformDLCStrong.IsValid())
+		{
+			PlatformDLCStrong->OnNotification().Remove(NotificationDelegate);
+		}
+		NotificationDelegate.Reset();
+	}
+
+	void OnNotification( FName DLCName, IPlatformDLC::ENotification Notification, bool bSuccess )
+	{
+		check(IsInGameThread());
+
+		if (Notification == IPlatformDLC::ENotification::Mounted && bSuccess)
+		{
+			TSharedPtr<IPlatformDLC> PlatformDLCStrong = PlatformDLC.Pin();
+			if (PlatformDLCStrong.IsValid())
+			{
+				FString MountPoint = PlatformDLCStrong->GetRootDirectory(DLCName);
+
+				bool bHasMountPoints = PlatformDLCFile->HasMountPoints();
+				PlatformDLCFile->AddMountPoint(DLCName, MountPoint);
+				if (!bHasMountPoints && PlatformDLCFile->HasMountPoints())
+				{
+					RegisterPlatformFile();
+				}
+			}
+		}
+		else if (Notification == IPlatformDLC::ENotification::Unmounted) // note: ignoring bSuccess here because the game's intent is for the DLC to unmount
+		{
+			bool bHasMountPoints = PlatformDLCFile->HasMountPoints();
+			PlatformDLCFile->RemoveMountPoint(DLCName);
+			if (bHasMountPoints && !PlatformDLCFile->HasMountPoints())
+			{
+				UnregisterPlatformFile();
+			}
+		}
+	};
+
 	bool bHasTried = false;
 	TWeakPtr<IPlatformDLC> PlatformDLC;
+	TUniquePtr<IPlatformDLCPlatformFile> PlatformDLCFile;
+	FDelegateHandle NotificationDelegate;
+
 };
 
 
@@ -118,7 +211,7 @@ static FAutoConsoleCommand CVarPlatformDLCDir(
 					class FDumpFileVistor : public IPlatformFile::FDirectoryVisitor
 					{
 					public:
-						virtual bool Visit(const TCHAR* FilenameOrDirectory, bool bIsDirectory)
+						virtual bool Visit(const TCHAR* FilenameOrDirectory, bool bIsDirectory) override
 						{
 							UE_CLOG(!bIsDirectory, LogPlatformDLC, Log, TEXT("%s"), FilenameOrDirectory);
 							return true;
diff --git a/Engine/Plugins/Experimental/PlatformDLC/Source/PlatformDLC/Private/PlatformDLCPlatformFile.cpp b/Engine/Plugins/Experimental/PlatformDLC/Source/PlatformDLC/Private/PlatformDLCPlatformFile.cpp
new file mode 100644
--- /dev/null
+++ b/Engine/Plugins/Experimental/PlatformDLC/Source/PlatformDLC/Private/PlatformDLCPlatformFile.cpp
@@ -0,0 +1,533 @@
+// Copyright Epic Games, Inc. All Rights Reserved.
+
+#include "PlatformDLCPlatformFile.h"
+#include "PlatformDLCModule.h"
+#include "Containers/DirectoryTree.h"
+#include "HAL/PlatformFile.h"
+#include "HAL/PlatformFileManager.h"
+#include "Misc/Paths.h"
+#include "Misc/ScopeRWLock.h"
+
+#if !defined(PLATFORMDLC_DUMP_LOOSE_FILES_ON_MOUNT)
+#define PLATFORMDLC_DUMP_LOOSE_FILES_ON_MOUNT (UE_BUILD_DEVELOPMENT || UE_BUILD_DEBUG)
+#endif
+
+namespace UE::PlatformDLC::Private
+{
+	/**
+	 * A file/directory visitor for files in DLC, used to share code for FDirectoryVisitor and FDirectoryStatVisitor
+	 * when iterating over files in DLC.
+	 */
+	class FDLCFileDirectoryVisitorBase
+	{
+	public:
+		virtual ~FDLCFileDirectoryVisitorBase() = default;
+		virtual bool ShouldVisitLeafPathname(FStringView LeafPathname) = 0;
+		virtual bool Visit(const FString& Filename, bool bIsDirectory, const FString* AbsPath) = 0;
+	};
+
+	/** FDLCFileDirectoryVisitorBase for a FDirectoryVisitor. */
+	class FDLCFileDirectoryVisitor : public FDLCFileDirectoryVisitorBase
+	{
+	public:
+		explicit FDLCFileDirectoryVisitor(IPlatformFile::FDirectoryVisitor& InInner)
+			: Inner(InInner) {}
+
+		virtual bool ShouldVisitLeafPathname(FStringView LeafPathname) override
+		{
+			return Inner.ShouldVisitLeafPathname(LeafPathname);
+		}
+		virtual bool Visit(const FString& Filename, bool bIsDirectory, const FString* AbsPath) override
+		{
+			return Inner.Visit(*Filename, bIsDirectory);
+		}
+
+		IPlatformFile::FDirectoryVisitor& Inner;
+	};
+
+	/** FDLCFileDirectoryVisitorBase for a FDirectoryStatVisitor. */
+	class FDLCFileDirectoryStatVisitor : public FDLCFileDirectoryVisitorBase
+	{
+	public:
+		FDLCFileDirectoryStatVisitor(IPlatformFile& InLowerLevel, IPlatformFile::FDirectoryStatVisitor& InInner)
+			: LowerLevel(InLowerLevel), Inner(InInner) {}
+
+		virtual bool ShouldVisitLeafPathname(FStringView LeafPathname) override
+		{
+			return Inner.ShouldVisitLeafPathname(LeafPathname);
+		}
+		virtual bool Visit(const FString& Filename, bool bIsDirectory, const FString* AbsPath) override
+		{
+			FFileStatData StatData;
+			if (!bIsDirectory && AbsPath)
+			{
+				StatData = LowerLevel.GetStatData(**AbsPath);
+			}
+			else if (bIsDirectory)
+			{
+				StatData.bIsValid = true;
+				StatData.bIsDirectory = true;
+				StatData.bIsReadOnly = true;
+				StatData.FileSize = -1;
+			}
+			return Inner.Visit(*Filename, StatData);
+		}
+
+		IPlatformFile& LowerLevel;
+		IPlatformFile::FDirectoryStatVisitor& Inner;
+	};
+
+}
+
+
+class FPlatformDLCPlatformFile : public IPlatformDLCPlatformFile
+{
+public:
+	FPlatformDLCPlatformFile() = default;
+	virtual ~FPlatformDLCPlatformFile() = default;
+
+	virtual bool Initialize(IPlatformFile* Inner, const TCHAR* CmdLine) override
+	{
+		LowerLevel = Inner;
+		return true;
+	}
+
+	virtual IPlatformFile* GetLowerLevel() override
+	{
+		return LowerLevel;
+	}
+
+	virtual void SetLowerLevel(IPlatformFile* NewLowerLevel) override
+	{
+		LowerLevel = NewLowerLevel;
+	}
+
+	virtual const TCHAR* GetName() const override
+	{
+		return TEXT("PlatformDLC");
+	}
+
+	virtual bool FileExists(const TCHAR* Filename) override
+	{
+		if (LowerLevel->FileExists(Filename))
+		{
+			return true;
+		}
+
+		FString AbsPath;
+		return FindFileInDLC(Filename, AbsPath);
+	}
+
+	virtual int64 FileSize(const TCHAR* Filename) override
+	{
+		if (LowerLevel->FileExists(Filename))
+		{
+			return LowerLevel->FileSize(Filename);
+		}
+
+		FString AbsPath;
+		if (FindFileInDLC(Filename, AbsPath))
+		{
+			return LowerLevel->FileSize(*AbsPath);
+		}
+		return -1;
+	}
+
+	virtual bool DeleteFile(const TCHAR* Filename) override
+	{
+		return LowerLevel->DeleteFile(Filename);
+	}
+
+	virtual bool IsReadOnly(const TCHAR* Filename) override
+	{
+		if (LowerLevel->FileExists(Filename))
+		{
+			return LowerLevel->IsReadOnly(Filename);
+		}
+
+		FString AbsPath;
+		return FindFileInDLC(Filename, AbsPath); // DLC files are always read-only
+	}
+
+	virtual bool MoveFile(const TCHAR* To, const TCHAR* From) override
+	{
+		return LowerLevel->MoveFile(To, From);
+	}
+
+	virtual bool SetReadOnly(const TCHAR* Filename, bool bNewReadOnlyValue) override
+	{
+		return LowerLevel->SetReadOnly(Filename, bNewReadOnlyValue);
+	}
+
+	virtual FDateTime GetTimeStamp(const TCHAR* Filename) override
+	{
+		if (LowerLevel->FileExists(Filename))
+		{
+			return LowerLevel->GetTimeStamp(Filename);
+		}
+
+		FString AbsPath;
+		if (FindFileInDLC(Filename, AbsPath))
+		{
+			return LowerLevel->GetTimeStamp(*AbsPath);
+		}
+		return FDateTime::MinValue();
+	}
+
+	virtual void SetTimeStamp(const TCHAR* Filename, FDateTime DateTime) override
+	{
+		LowerLevel->SetTimeStamp(Filename, DateTime);
+	}
+
+	virtual FDateTime GetAccessTimeStamp(const TCHAR* Filename) override
+	{
+		if (LowerLevel->FileExists(Filename))
+		{
+			return LowerLevel->GetAccessTimeStamp(Filename);
+		}
+
+		FString AbsPath;
+		if (FindFileInDLC(Filename, AbsPath))
+		{
+			return LowerLevel->GetAccessTimeStamp(*AbsPath);
+		}
+		return FDateTime::MinValue();
+	}
+
+	virtual FString GetFilenameOnDisk(const TCHAR* Filename) override
+	{
+		if (LowerLevel->FileExists(Filename))
+		{
+			return LowerLevel->GetFilenameOnDisk(Filename);
+		}
+
+		FString AbsPath;
+		if (FindFileInDLC(Filename, AbsPath))
+		{
+			return LowerLevel->GetFilenameOnDisk(*AbsPath);
+		}
+		return Filename;
+	}
+
+	virtual IFileHandle* OpenRead(const TCHAR* Filename, bool bAllowWrite = false) override
+	{
+		IFileHandle* Handle = LowerLevel->OpenRead(Filename, bAllowWrite);
+		if (Handle)
+		{
+			return Handle;
+		}
+
+		FString AbsPath;
+		if (FindFileInDLC(Filename, AbsPath))
+		{
+			return LowerLevel->OpenRead(*AbsPath, bAllowWrite);
+		}
+		return nullptr;
+	}
+
+	virtual IFileHandle* OpenWrite(const TCHAR* Filename, bool bAppend = false, bool bAllowRead = false) override
+	{
+		return LowerLevel->OpenWrite(Filename, bAppend, bAllowRead);
+	}
+
+	virtual FFileStatData GetStatData(const TCHAR* FilenameOrDirectory) override
+	{
+		FFileStatData StatData = LowerLevel->GetStatData(FilenameOrDirectory);
+		if (StatData.bIsValid)
+		{
+			return StatData;
+		}
+
+		FString AbsPath;
+		if (FindFileInDLC(FilenameOrDirectory, AbsPath))
+		{
+			return LowerLevel->GetStatData(*AbsPath);
+		}
+		if (IsDLCDirectory(FilenameOrDirectory))
+		{
+			FFileStatData DirStat;
+			DirStat.bIsValid = true;
+			DirStat.bIsDirectory = true;
+			DirStat.bIsReadOnly = true;
+			DirStat.FileSize = -1;
+			return DirStat;
+		}
+		return StatData;
+	}
+
+	virtual bool DirectoryExists(const TCHAR* Directory) override
+	{
+		if (LowerLevel->DirectoryExists(Directory))
+		{
+			return true;
+		}
+
+		return IsDLCDirectory(Directory);
+	}
+
+	virtual bool CreateDirectory(const TCHAR* Directory) override
+	{
+		return LowerLevel->CreateDirectory(Directory);
+	}
+
+	virtual bool DeleteDirectory(const TCHAR* Directory) override
+	{
+		return LowerLevel->DeleteDirectory(Directory);
+	}
+
+	virtual bool IterateDirectory(const TCHAR* Directory, FDirectoryVisitor& Visitor) override
+	{
+		bool bResult = LowerLevel->IterateDirectory(Directory, Visitor);
+
+		UE::PlatformDLC::Private::FDLCFileDirectoryVisitor DLCVisitor(Visitor);
+		return IterateDirectoryInDLCFiles(Directory, DLCVisitor) || bResult;
+	}
+
+	virtual bool IterateDirectoryStat(const TCHAR* Directory, FDirectoryStatVisitor& Visitor) override
+	{
+		bool bResult = LowerLevel->IterateDirectoryStat(Directory, Visitor);
+
+		UE::PlatformDLC::Private::FDLCFileDirectoryStatVisitor DLCStatVisitor(*LowerLevel, Visitor);
+		return IterateDirectoryInDLCFiles(Directory, DLCStatVisitor) || bResult;
+	}
+
+
+protected:
+	virtual void AddMountPoint( FName DLCName, const FString& MountPoint ) override
+	{
+		// early out if we have this already
+		{
+			FReadScopeLock Lock(MountPointsLock);
+			if (MountPoints.Contains(DLCName))
+			{
+				return;
+			}
+		}
+
+		TSharedPtr<FDLCFiles> DLCFiles = MakeShared<FDLCFiles>();
+		DLCFiles->MountPoint = MountPoint;
+
+		// DLC root directory is often an unusual path outside of the game's root
+		FString NormalizedDLCRoot = MountPoint;
+		FPaths::NormalizeDirectoryName(NormalizedDLCRoot);
+		NormalizedDLCRoot += TEXT("/"); // MakePathRelativeTo expects a trailing / on the path
+
+		// enumerate all the files in the DLC, ignoring anything in the pak folder
+		TArray<FString> FoundFiles;
+		IPlatformFile::GetPlatformPhysical().FindFilesRecursively(FoundFiles, *MountPoint, TEXT(""));
+		for (FString& File : FoundFiles)
+		{
+			FString Directory = FPaths::GetPath(File);
+			if (Directory.EndsWith(TEXT("/Paks"), ESearchCase::IgnoreCase))
+			{
+				continue;
+			}
+
+			FString RelativePath = File;
+			FPaths::MakePathRelativeTo(RelativePath, *NormalizedDLCRoot);
+
+			DLCFiles->Files.FindOrAdd(RelativePath) = true;
+		}
+		DLCFiles->Files.Shrink();
+
+		// save this DLC if we found loose files
+		if (!DLCFiles->Files.IsEmpty())
+		{
+			UE_LOG(LogPlatformDLC, Log, TEXT("DLC %s loose file tree is %llu bytes for %d files"), *DLCName.ToString(), DLCFiles->Files.GetAllocatedSize(), DLCFiles->Files.Num());
+
+			FWriteScopeLock Lock(MountPointsLock);
+			MountPoints.Add(DLCName, MoveTemp(DLCFiles));
+		}
+
+#if PLATFORMDLC_DUMP_LOOSE_FILES_ON_MOUNT
+		{
+			FReadScopeLock Lock(MountPointsLock);
+			UE_LOG(LogPlatformDLC, Log, TEXT("%d mounted DLCs. Unified loose files are:"), MountPoints.Num() );
+		}
+		DumpLooseFiles();
+#endif
+	}
+
+	virtual void RemoveMountPoint( FName DLCName ) override
+	{
+		FWriteScopeLock Lock(MountPointsLock);
+		MountPoints.Remove(DLCName);
+	}
+
+	virtual bool HasMountPoints() const override
+	{
+		FReadScopeLock Lock(MountPointsLock);
+		return MountPoints.Num() > 0;
+	}
+
+private:
+	struct FDLCFiles
+	{
+		FString MountPoint;
+		TDirectoryTree<bool> Files; // relative to game root
+	};
+
+	bool IterateDirectoryInDLCFiles(const TCHAR* Directory, UE::PlatformDLC::Private::FDLCFileDirectoryVisitorBase& Visitor)
+	{
+		bool bResult = false;
+
+		// early out for unsupported paths
+		if (!FStringView(Directory).StartsWith(TEXTVIEW("../../../")))
+		{
+			return bResult;
+		}
+		FString RelDir = MakeRelativeToDLCRoot(Directory);
+
+		// enumerate paths from all mounted DLC
+		TSet<FString> ChildPaths;
+		{
+			FReadScopeLock Lock(MountPointsLock);
+			for (const TPair<FName, TSharedPtr<FDLCFiles>>& Pair : MountPoints)
+			{
+				const FDLCFiles& DLCFiles = *Pair.Value;
+				TArray<FString> ChildNames;
+				if (!DLCFiles.Files.TryGetChildren(RelDir, ChildNames, EDirectoryTreeGetFlags::ImpliedParent | EDirectoryTreeGetFlags::ImpliedChildren))
+				{
+					continue;
+				}
+
+				bResult = true;
+				for (const FString& ChildName : ChildNames)
+				{
+					if (!Visitor.ShouldVisitLeafPathname(ChildName))
+					{
+						continue;
+					}
+
+					FString ChildPath = FPaths::Combine(Directory, *ChildName);
+					FString ChildRelPath = FPaths::Combine(RelDir, ChildName);
+
+					if (DLCFiles.Files.ContainsChildPaths(ChildRelPath))
+					{
+						ChildPaths.Add(ChildPath);
+					}
+					else
+					{
+						FString AbsChildPath = FPaths::Combine(DLCFiles.MountPoint, ChildRelPath);
+						if (!Visitor.Visit(ChildPath, false, &AbsChildPath))
+						{
+							return false;
+						}
+					}
+				}
+			}
+		}
+
+		// must iterate paths separately because the same directory is likely to exist in multiple DLCs (../../../MyGame/Content/ for example)
+		for (const FString& ChildPath : ChildPaths)
+		{
+			if (!Visitor.Visit(ChildPath, true, nullptr))
+			{
+				return false;
+			}
+		}
+
+		return bResult;
+	}
+
+	static FString MakeRelativeToDLCRoot(const TCHAR* InPath)
+	{
+		FString Path = InPath;
+		if (!Path.RemoveFromStart(TEXT("../../../")))
+		{
+			return FString();
+		}
+		return Path;
+	}
+
+	bool FindFileInDLC(const TCHAR* InPath, FString& OutAbsPath) const
+	{
+		FString RelPath = MakeRelativeToDLCRoot(InPath);
+		if (RelPath.IsEmpty())
+		{
+			return false;
+		}
+
+		FReadScopeLock Lock(MountPointsLock);
+		for (const TPair<FName, TSharedPtr<FDLCFiles>>& Pair : MountPoints)
+		{
+			if (Pair.Value->Files.Contains(RelPath))
+			{
+				OutAbsPath = FPaths::Combine(Pair.Value->MountPoint, RelPath);
+				return true;
+			}
+		}
+		return false;
+	}
+
+	bool IsDLCDirectory(const TCHAR* InPath) const
+	{
+		FString RelPath = MakeRelativeToDLCRoot(InPath);
+		if (RelPath.IsEmpty())
+		{
+			return false;
+		}
+
+		FReadScopeLock Lock(MountPointsLock);
+		for (const TPair<FName, TSharedPtr<FDLCFiles>>& Pair : MountPoints)
+		{
+			if (Pair.Value->Files.ContainsChildPaths(RelPath))
+			{
+				return true;
+			}
+		}
+		return false;
+	}
+
+
+#if PLATFORMDLC_DUMP_LOOSE_FILES_ON_MOUNT
+	void DumpLooseFiles(const TCHAR* Directory = TEXT("../../../"))
+	{
+		using namespace UE::PlatformDLC::Private;
+
+		class FDumpVisitor : public IPlatformFile::FDirectoryVisitor
+		{
+		public:
+			TArray<FString> Subdirectories;
+
+			virtual bool Visit(const TCHAR* FilenameOrDirectory, bool bIsDirectory) override
+			{
+				if (bIsDirectory)
+				{
+					Subdirectories.Add(FilenameOrDirectory);
+				}
+				else
+				{
+					UE_LOG(LogPlatformDLC, Log, TEXT("  %s"), FilenameOrDirectory);
+				}
+				return true;
+			}
+		};
+
+		FDumpVisitor InnerVisitor;
+		FDLCFileDirectoryVisitor DLCVisitor(InnerVisitor);
+		IterateDirectoryInDLCFiles(Directory, DLCVisitor);
+
+		for (const FString& Subdir : InnerVisitor.Subdirectories)
+		{
+			DumpLooseFiles(*Subdir);
+		}
+	}
+#endif
+
+
+
+	IPlatformFile* LowerLevel = nullptr;
+
+	mutable FRWLock MountPointsLock;
+	TMap<FName, TSharedPtr<FDLCFiles>> MountPoints;
+};
+
+
+
+
+TUniquePtr<IPlatformDLCPlatformFile> IPlatformDLCPlatformFile::Construct()
+{
+	return MakeUnique<FPlatformDLCPlatformFile>();
+}
+
diff --git a/Engine/Plugins/Experimental/PlatformDLC/Source/PlatformDLC/Private/PlatformDLCPlatformFile.h b/Engine/Plugins/Experimental/PlatformDLC/Source/PlatformDLC/Private/PlatformDLCPlatformFile.h
new file mode 100644
--- /dev/null
+++ b/Engine/Plugins/Experimental/PlatformDLC/Source/PlatformDLC/Private/PlatformDLCPlatformFile.h
@@ -0,0 +1,19 @@
+// Copyright Epic Games, Inc. All Rights Reserved.
+
+#pragma once
+
+#include "CoreMinimal.h"
+#include "GenericPlatform/GenericPlatformFile.h"
+
+class IPlatformDLCPlatformFile : public IPlatformFile
+{
+public:
+	static TUniquePtr<IPlatformDLCPlatformFile> Construct();
+
+	virtual ~IPlatformDLCPlatformFile() = default;
+
+	// IPlatformDLCPlatformFile interface
+	virtual void AddMountPoint( FName DLCName, const FString& MountPoint ) = 0;
+	virtual void RemoveMountPoint( FName DLCName ) = 0;
+	virtual bool HasMountPoints() const = 0;
+};
diff --git a/Engine/Plugins/Experimental/PlatformDLC/Source/PlatformDLC/Public/PlatformDLCModule.h b/Engine/Plugins/Experimental/PlatformDLC/Source/PlatformDLC/Public/PlatformDLCModule.h
--- a/Engine/Plugins/Experimental/PlatformDLC/Source/PlatformDLC/Public/PlatformDLCModule.h
+++ b/Engine/Plugins/Experimental/PlatformDLC/Source/PlatformDLC/Public/PlatformDLCModule.h
@@ -4,13 +4,14 @@
 
 #include "CoreMinimal.h"
 #include "PlatformDLC.h"
+#include "HAL/IPlatformFileModule.h"
 #include "Modules/ModuleInterface.h"
 #include "Modules/ModuleManager.h"
 #include "Templates/SharedPointer.h"
 
 PLATFORMDLC_API DECLARE_LOG_CATEGORY_EXTERN(LogPlatformDLC, Log, All)
 
-class IPlatformDLCModule : public IModuleInterface
+class IPlatformDLCModule : public IPlatformFileModule
 {
 public:
 	/**
