Add a LexToString for IPlatformDLC ENotifications

#rnx

#Synced-CL 52105125

diff --git a/Engine/Plugins/Experimental/PlatformDLC/Source/PlatformDLC/Private/PlatformDLCModule.cpp b/Engine/Plugins/Experimental/PlatformDLC/Source/PlatformDLC/Private/PlatformDLCModule.cpp
--- a/Engine/Plugins/Experimental/PlatformDLC/Source/PlatformDLC/Private/PlatformDLCModule.cpp
+++ b/Engine/Plugins/Experimental/PlatformDLC/Source/PlatformDLC/Private/PlatformDLCModule.cpp
@@ -255,4 +255,18 @@ static FAutoConsoleCommand CVarPlatformDLCList(
 
 
 
+FString LexToString(IPlatformDLC::ENotification Notification)
+{
+	switch (Notification)
+	{
+		case IPlatformDLC::ENotification::Entitlement:  return TEXT("Entitlement");
+		case IPlatformDLC::ENotification::Mounted:      return TEXT("Mounted");
+		case IPlatformDLC::ENotification::Unmounted:    return TEXT("Unmounted");
+		case IPlatformDLC::ENotification::Downloaded:   return TEXT("Downloaded");
+		case IPlatformDLC::ENotification::Uninstalled:  return TEXT("Uninstalled");
+		default:                                        return TEXT("Invalid");
+	}
+}
+
+
 IMPLEMENT_MODULE(FPlatformDLCModule, PlatformDLC);
diff --git a/Engine/Plugins/Experimental/PlatformDLC/Source/PlatformDLC/Public/PlatformDLC.h b/Engine/Plugins/Experimental/PlatformDLC/Source/PlatformDLC/Public/PlatformDLC.h
--- a/Engine/Plugins/Experimental/PlatformDLC/Source/PlatformDLC/Public/PlatformDLC.h
+++ b/Engine/Plugins/Experimental/PlatformDLC/Source/PlatformDLC/Public/PlatformDLC.h
@@ -95,3 +95,7 @@ public:
 	// get the platform DLC for this factory
 	virtual TSharedPtr<IPlatformDLC> GetPlatformDLC() = 0;
 };
+
+
+PLATFORMDLC_API FString LexToString(IPlatformDLC::ENotification Notification);
+
