Add extra IPlatformFile functions to PlatformDLCFile

#rnx
#Synced-CL 52105125

diff --git a/Engine/Plugins/Experimental/PlatformDLC/Source/PlatformDLC/Private/PlatformDLCPlatformFile.cpp b/Engine/Plugins/Experimental/PlatformDLC/Source/PlatformDLC/Private/PlatformDLCPlatformFile.cpp
--- a/Engine/Plugins/Experimental/PlatformDLC/Source/PlatformDLC/Private/PlatformDLCPlatformFile.cpp
+++ b/Engine/Plugins/Experimental/PlatformDLC/Source/PlatformDLC/Private/PlatformDLCPlatformFile.cpp
@@ -5,6 +5,7 @@
 #include "Containers/DirectoryTree.h"
 #include "HAL/PlatformFile.h"
 #include "HAL/PlatformFileManager.h"
+#include "Async/MappedFileHandle.h"
 #include "Misc/Paths.h"
 #include "Misc/ScopeRWLock.h"
 
@@ -291,6 +292,89 @@ public:
 		return IterateDirectoryInDLCFiles(Directory, DLCStatVisitor) || bResult;
 	}
 
+	//~ For visibility of overloads we don't override
+	using IPlatformFile::IterateDirectory;
+	using IPlatformFile::IterateDirectoryRecursively;
+	using IPlatformFile::IterateDirectoryStat;
+	using IPlatformFile::IterateDirectoryStatRecursively;
+	using IPlatformFile::OpenAsyncRead;
+
+	virtual bool DeleteDirectoryRecursively(const TCHAR* Directory) override
+	{
+		return LowerLevel->DeleteDirectoryRecursively(Directory);
+	}
+
+	virtual bool CopyFile(const TCHAR* To, const TCHAR* From, EPlatformFileRead ReadFlags = EPlatformFileRead::None, EPlatformFileWrite WriteFlags = EPlatformFileWrite::None) override
+	{
+		return LowerLevel->CopyFile(To, From, ReadFlags, WriteFlags);
+	}
+
+	virtual bool CreateDirectoryTree(const TCHAR* Directory) override
+	{
+		return LowerLevel->CreateDirectoryTree(Directory);
+	}
+
+	virtual bool CopyDirectoryTree(const TCHAR* DestinationDirectory, const TCHAR* Source, bool bOverwriteAllExisting) override
+	{
+		return LowerLevel->CopyDirectoryTree(DestinationDirectory, Source, bOverwriteAllExisting);
+	}
+
+	virtual FString ConvertToAbsolutePathForExternalAppForRead(const TCHAR* Filename) override
+	{
+		if (LowerLevel->FileExists(Filename))
+		{
+			return LowerLevel->ConvertToAbsolutePathForExternalAppForRead(Filename);
+		}
+
+		FString AbsPath;
+		if (FindFileInDLC(Filename, AbsPath))
+		{
+			return AbsPath;
+		}
+
+		return LowerLevel->ConvertToAbsolutePathForExternalAppForRead(Filename);
+	}
+
+	virtual FString ConvertToAbsolutePathForExternalAppForWrite(const TCHAR* Filename) override
+	{
+		return LowerLevel->ConvertToAbsolutePathForExternalAppForWrite(Filename);
+	}
+
+	virtual IAsyncReadFileHandle* OpenAsyncRead(const TCHAR* Filename, bool bAllowWrite = false) override
+	{
+		if (!LowerLevel->FileExists(Filename))
+		{
+			FString AbsPath;
+			if (FindFileInDLC(Filename, AbsPath))
+			{
+				return LowerLevel->OpenAsyncRead(*AbsPath, bAllowWrite);
+			}
+		}
+		return LowerLevel->OpenAsyncRead(Filename, bAllowWrite);
+	}
+
+	virtual FOpenMappedResult OpenMappedEx(const TCHAR* Filename, EOpenReadFlags OpenOptions = EOpenReadFlags::None, int64 MaximumSize = 0) override
+	{
+		FOpenMappedResult Result = LowerLevel->OpenMappedEx(Filename, OpenOptions, MaximumSize);
+		if (Result.HasValue())
+		{
+			return Result;
+		}
+
+		FString AbsPath;
+		if (FindFileInDLC(Filename, AbsPath))
+		{
+			return LowerLevel->OpenMappedEx(*AbsPath, OpenOptions, MaximumSize);
+		}
+		return Result;
+	}
+
+	virtual bool SendMessageToServer(const TCHAR* Message, IFileServerMessageHandler* Handler) override
+	{
+		return LowerLevel->SendMessageToServer(Message, Handler);
+	}
+
+
 
 protected:
 	virtual void AddMountPoint( FName DLCName, const FString& MountPoint ) override
